Fix OculusOpenXR manifest to have correct application name. This fixes OpenXR app on Quest not opening as a VR application.
#rnx

diff --git a/Engine/Plugins/Runtime/Oculus/OculusOpenXR/Source/OculusOpenXRHMD/OculusMobile_APL.xml b/Engine/Plugins/Runtime/Oculus/OculusOpenXR/Source/OculusOpenXRHMD/OculusMobile_APL.xml
--- a/Engine/Plugins/Runtime/Oculus/OculusOpenXR/Source/OculusOpenXRHMD/OculusMobile_APL.xml
+++ b/Engine/Plugins/Runtime/Oculus/OculusOpenXR/Source/OculusOpenXRHMD/OculusMobile_APL.xml
@@ -134,7 +134,7 @@
 		<!-- Add Activity Specific Flags -->
 		<loopElements tag="activity">
 			<setStringFromAttribute result="activityName" tag="$" name="android:name"/>
-			<setBoolIsEqual result="bGameActivity" arg1="$S(activityName)" arg2="com.epicgames.ue4.GameActivity"/>
+			<setBoolIsEqual result="bGameActivity" arg1="$S(activityName)" arg2="com.epicgames.unreal.GameActivity"/>
 			<if condition="bGameActivity">
 				<true>
 					<!-- Add VR Intent Filter, Permissions, and Features -->
